remove #pragma once from cpp

diff --git a/Engine/Plugins/Animation/LiveLinkCurveDebugUI/Source/LiveLinkCurveDebugUI/Private/LiveLinkDebuggerSettings.cpp b/Engine/Plugins/Animation/LiveLinkCurveDebugUI/Source/LiveLinkCurveDebugUI/Private/LiveLinkDebuggerSettings.cpp
--- a/Engine/Plugins/Animation/LiveLinkCurveDebugUI/Source/LiveLinkCurveDebugUI/Private/LiveLinkDebuggerSettings.cpp
+++ b/Engine/Plugins/Animation/LiveLinkCurveDebugUI/Source/LiveLinkCurveDebugUI/Private/LiveLinkDebuggerSettings.cpp
@@ -1,7 +1,5 @@
 // Copyright 1998-2018 Epic Games, Inc. All Rights Reserved.
 
-#pragma once
-
 #include "LiveLinkDebuggerSettings.h"
 
 ULiveLinkDebuggerSettings::ULiveLinkDebuggerSettings()
