DR - Linux compile fix
#rnx

diff --git a/Engine/Plugins/Animation/LiveLinkCurveDebugUI/Source/LiveLinkCurveDebugUI/Private/LiveLinkDebuggerSettings.cpp b/Engine/Plugins/Animation/LiveLinkCurveDebugUI/Source/LiveLinkCurveDebugUI/Private/LiveLinkDebuggerSettings.cpp
--- a/Engine/Plugins/Animation/LiveLinkCurveDebugUI/Source/LiveLinkCurveDebugUI/Private/LiveLinkDebuggerSettings.cpp
+++ b/Engine/Plugins/Animation/LiveLinkCurveDebugUI/Source/LiveLinkCurveDebugUI/Private/LiveLinkDebuggerSettings.cpp
@@ -2,8 +2,6 @@
 
 #pragma once
 
-#include "LiveLinkDebuggerSettings.h"
-
 ULiveLinkDebuggerSettings::ULiveLinkDebuggerSettings()
 	: MinBarColor(FLinearColor(.05f, 0, 0, 1))
 	, MaxBarColor(FLinearColor(1.f, 0, 0, 1))
