DR - Fix the fix
#rnx

diff --git a/Engine/Plugins/Animation/LiveLinkCurveDebugUI/Source/LiveLinkCurveDebugUI/Private/LiveLinkDebuggerSettings.cpp b/Engine/Plugins/Animation/LiveLinkCurveDebugUI/Source/LiveLinkCurveDebugUI/Private/LiveLinkDebuggerSettings.cpp
--- a/Engine/Plugins/Animation/LiveLinkCurveDebugUI/Source/LiveLinkCurveDebugUI/Private/LiveLinkDebuggerSettings.cpp
+++ b/Engine/Plugins/Animation/LiveLinkCurveDebugUI/Source/LiveLinkCurveDebugUI/Private/LiveLinkDebuggerSettings.cpp
@@ -1,5 +1,7 @@
 // Copyright 1998-2018 Epic Games, Inc. All Rights Reserved.
 
+#include "LiveLinkDebuggerSettings.h"
+
 ULiveLinkDebuggerSettings::ULiveLinkDebuggerSettings()
 	: MinBarColor(FLinearColor(.05f, 0, 0, 1))
 	, MaxBarColor(FLinearColor(1.f, 0, 0, 1))
