Merging //UE4/Dev-Main to Dev-Mobile (//UE4/Dev-Mobile)

diff --git a/Engine/Plugins/Animation/LiveLinkCurveDebugUI/Source/LiveLinkCurveDebugUI/LiveLinkCurveDebugUI.Build.cs b/Engine/Plugins/Animation/LiveLinkCurveDebugUI/Source/LiveLinkCurveDebugUI/LiveLinkCurveDebugUI.Build.cs
--- a/Engine/Plugins/Animation/LiveLinkCurveDebugUI/Source/LiveLinkCurveDebugUI/LiveLinkCurveDebugUI.Build.cs
+++ b/Engine/Plugins/Animation/LiveLinkCurveDebugUI/Source/LiveLinkCurveDebugUI/LiveLinkCurveDebugUI.Build.cs
@@ -1,4 +1,4 @@
-// Copyright 1998-2018 Epic Games, Inc. All Rights Reserved.
+// Copyright 1998-2019 Epic Games, Inc. All Rights Reserved.
 
 namespace UnrealBuildTool.Rules
 {
