Merge down compilation fix.
#rnx

diff --git a/Engine/Plugins/Animation/LiveLinkCurveDebugUI/Source/LiveLinkCurveDebugUI/Public/LiveLinkDebugCurveNodeBase.h b/Engine/Plugins/Animation/LiveLinkCurveDebugUI/Source/LiveLinkCurveDebugUI/Public/LiveLinkDebugCurveNodeBase.h
--- a/Engine/Plugins/Animation/LiveLinkCurveDebugUI/Source/LiveLinkCurveDebugUI/Public/LiveLinkDebugCurveNodeBase.h
+++ b/Engine/Plugins/Animation/LiveLinkCurveDebugUI/Source/LiveLinkCurveDebugUI/Public/LiveLinkDebugCurveNodeBase.h
@@ -4,6 +4,8 @@
 
 #include "CoreMinimal.h"
 
+struct FSlateColor;
+
 class FLiveLinkDebugCurveNodeBase
 {
 public:
