Fixed some issues not found from PF when compiling with -allmodules.

diff --git a/Engine/Plugins/Animation/LiveLinkCurveDebugUI/Source/LiveLinkCurveDebugUI/LiveLinkCurveDebugUI.Build.cs b/Engine/Plugins/Animation/LiveLinkCurveDebugUI/Source/LiveLinkCurveDebugUI/LiveLinkCurveDebugUI.Build.cs
--- a/Engine/Plugins/Animation/LiveLinkCurveDebugUI/Source/LiveLinkCurveDebugUI/LiveLinkCurveDebugUI.Build.cs
+++ b/Engine/Plugins/Animation/LiveLinkCurveDebugUI/Source/LiveLinkCurveDebugUI/LiveLinkCurveDebugUI.Build.cs
@@ -7,16 +7,23 @@ namespace UnrealBuildTool.Rules
 		public LiveLinkCurveDebugUI(ReadOnlyTargetRules Target) : base(Target)
 		{
 			PrivateDependencyModuleNames.AddRange(
-			new string[]
-			{
-				"Core",
-				"CoreUObject",
-				"Engine",
-				"LiveLinkInterface",
-				"Slate",
-				"SlateCore",
-				"InputCore"
-			}
+				new string[]
+				{
+					"Core",
+					"CoreUObject",
+					"Engine",
+					"LiveLinkInterface",
+					"Slate",
+					"SlateCore",
+					"InputCore"
+				}
+			);
+
+			PrivateIncludePathModuleNames.AddRange(
+				new string[]
+				{
+					"WorkspaceMenuStructure"
+				}
 			);
 
 			// DesktopPlatform is only available for Editor and Program targets (running on a desktop platform)
diff --git a/Engine/Plugins/Compositing/Composure/Source/Composure/Private/ComposureBlueprintLibrary.cpp b/Engine/Plugins/Compositing/Composure/Source/Composure/Private/ComposureBlueprintLibrary.cpp
--- a/Engine/Plugins/Compositing/Composure/Source/Composure/Private/ComposureBlueprintLibrary.cpp
+++ b/Engine/Plugins/Compositing/Composure/Source/Composure/Private/ComposureBlueprintLibrary.cpp
@@ -7,8 +7,8 @@
 #include "CameraCalibrationSubsystem.h"
 #include "CineCameraComponent.h"
 #include "Components/SceneCaptureComponent2D.h"
-#include "ICompElementManager.h"
-#include "CompElementEditorModule.h"
+#include "ComposureLayersEditor/Public/ICompElementManager.h"
+#include "ComposureLayersEditor/Public/CompElementEditorModule.h"
 #include "ComposurePlayerCompositingTarget.h"
 #include "ComposureUtils.h"
 #include "Engine/LocalPlayer.h"
diff --git a/Engine/Plugins/Editor/SpeedTreeImporter/Source/SpeedTreeImporter/Private/SpeedTreeImportFactory.cpp b/Engine/Plugins/Editor/SpeedTreeImporter/Source/SpeedTreeImporter/Private/SpeedTreeImportFactory.cpp
--- a/Engine/Plugins/Editor/SpeedTreeImporter/Source/SpeedTreeImporter/Private/SpeedTreeImportFactory.cpp
+++ b/Engine/Plugins/Editor/SpeedTreeImporter/Source/SpeedTreeImporter/Private/SpeedTreeImportFactory.cpp
@@ -64,7 +64,7 @@
 #include "PackageTools.h"
 #include "AssetRegistry/AssetRegistryModule.h"
 #include "AssetToolsModule.h"
-#include "GeomFitUtils.h"
+#include "UnrealEd/Private/GeomFitUtils.h"
 #include "SpeedTreeWind.h"
 #include "StaticMeshAttributes.h"
 #include "ComponentReregisterContext.h"
@@ -72,9 +72,9 @@
 #if WITH_SPEEDTREE
 
 THIRD_PARTY_INCLUDES_START
-#include "ThirdParty/SpeedTree/SpeedTreeSDK-v7.0/Include/Core/Core.h"
-#include "ThirdParty/SpeedTree/SpeedTreeDataBuffer/TreeReader8.h"
-#include "ThirdParty/SpeedTree/SpeedTreeDataBuffer/TreeReader9.h"
+#include "Core/Core.h"
+#include "TreeReader8.h"
+#include "TreeReader9.h"
 THIRD_PARTY_INCLUDES_END
 
 #endif // WITH_SPEEDTREE
